Input debugging plugin.

 - includes touch input visualization, toggled via `Input.Debug.ShowTouches 1`

diff --git a/Engine/Plugins/Runtime/InputDebugging/InputDebugging.uplugin b/Engine/Plugins/Runtime/InputDebugging/InputDebugging.uplugin
new file mode 100644
--- /dev/null
+++ b/Engine/Plugins/Runtime/InputDebugging/InputDebugging.uplugin
@@ -0,0 +1,24 @@
+{
+	"FileVersion": 3,
+	"Version": 1,
+	"VersionName": "1.0",
+	"FriendlyName": "Input Debugging",
+	"Description": "Input debugging and visualization.",
+	"Category": "Input",
+	"CreatedBy": "Epic Games, Inc.",
+	"CreatedByURL": "http://epicgames.com",
+	"DocsURL": "",
+	"MarketplaceURL": "",
+	"SupportURL": "",
+	"EnabledByDefault": true,
+	"CanContainContent": false,
+	"IsBetaVersion": false,
+	"Installed": false,
+	"Modules": [
+		{
+			"Name": "InputDebugging",
+			"Type": "Runtime",
+			"LoadingPhase": "Default"
+		}
+	]
+}
\ No newline at end of file
diff --git a/Engine/Plugins/Runtime/InputDebugging/Source/InputDebugging/InputDebugging.Build.cs b/Engine/Plugins/Runtime/InputDebugging/Source/InputDebugging/InputDebugging.Build.cs
new file mode 100644
--- /dev/null
+++ b/Engine/Plugins/Runtime/InputDebugging/Source/InputDebugging/InputDebugging.Build.cs
@@ -0,0 +1,19 @@
+// Copyright Epic Games, Inc. All Rights Reserved.
+
+namespace UnrealBuildTool.Rules
+{
+	public class InputDebugging : ModuleRules
+	{
+		public InputDebugging(ReadOnlyTargetRules Target) : base(Target)
+		{
+			PrivateDependencyModuleNames.AddRange(
+				new string[] {
+					"Core",
+					"Engine",
+					"Slate",
+					"InputCore",
+				}
+			);
+		}
+	}
+}
\ No newline at end of file
diff --git a/Engine/Plugins/Runtime/InputDebugging/Source/InputDebugging/Private/InputDebuggingModule.cpp b/Engine/Plugins/Runtime/InputDebugging/Source/InputDebugging/Private/InputDebuggingModule.cpp
new file mode 100644
--- /dev/null
+++ b/Engine/Plugins/Runtime/InputDebugging/Source/InputDebugging/Private/InputDebuggingModule.cpp
@@ -0,0 +1,61 @@
+// Copyright Epic Games, Inc. All Rights Reserved.
+
+#include "IInputDebuggingInterface.h"
+#include "Debug/DebugDrawService.h"
+#include "Engine/Canvas.h"
+#include "TouchInputVisualizer.h"
+
+
+class FInputDebuggingModule : public IInputDebuggingInterface
+{
+public:
+
+	virtual void StartupModule() override
+	{
+		// register modular feature
+		check(!IInputDebuggingInterface::IsAvailable()); // there can be only one
+		IModularFeatures::Get().RegisterModularFeature(IInputDebuggingInterface::GetModularFeatureName(), this);
+
+		// register debug rendering callback
+		DrawOnCanvasDelegateHandle = UDebugDrawService::Register(TEXT("InputDebugVisualizer"), FDebugDrawDelegate::CreateRaw(this, &FInputDebuggingModule::OnDebugDraw));
+
+		// initialize the touch input visualizer
+#if SUPPORT_TOUCH_INPUT_DISPLAY
+		TouchInputVisualizer = MakeShared<FTouchInputVisualizer>();
+#endif//SUPPORT_TOUCH_INPUT_DISPLAY
+	}
+
+	virtual void ShutdownModule() override
+	{
+		// shutdown the touch input visualizer
+#if SUPPORT_TOUCH_INPUT_DISPLAY
+		TouchInputVisualizer.Reset();
+#endif//SUPPORT_TOUCH_INPUT_DISPLAY
+
+		// unregister debug rendering callback
+		UDebugDrawService::Unregister(DrawOnCanvasDelegateHandle);
+
+		// unregister modular feature
+		IModularFeatures::Get().UnregisterModularFeature(IInputDebuggingInterface::GetModularFeatureName(), this);
+	}
+
+
+	void OnDebugDraw(UCanvas* Canvas, class APlayerController* PlayerController)
+	{
+		// debug render the touch input visualization.
+#if SUPPORT_TOUCH_INPUT_DISPLAY
+		TouchInputVisualizer->OnDebugDraw(Canvas);
+#endif //SUPPORT_TOUCH_INPUT_DISPLAY
+	}
+
+private:
+
+	FDelegateHandle DrawOnCanvasDelegateHandle;
+
+#if SUPPORT_TOUCH_INPUT_DISPLAY
+	TSharedPtr<FTouchInputVisualizer> TouchInputVisualizer;
+#endif//SUPPORT_TOUCH_INPUT_DISPLAY
+
+};
+
+IMPLEMENT_MODULE( FInputDebuggingModule, InputDebugging);
diff --git a/Engine/Plugins/Runtime/InputDebugging/Source/InputDebugging/Private/TouchInputVisualizer.cpp b/Engine/Plugins/Runtime/InputDebugging/Source/InputDebugging/Private/TouchInputVisualizer.cpp
new file mode 100644
--- /dev/null
+++ b/Engine/Plugins/Runtime/InputDebugging/Source/InputDebugging/Private/TouchInputVisualizer.cpp
@@ -0,0 +1,112 @@
+// Copyright Epic Games, Inc. All Rights Reserved.
+
+#include "TouchInputVisualizer.h"
+#include "CoreMinimal.h"
+#include "Engine/Canvas.h"
+#include "Misc/CoreDelegates.h"
+#include "HAL/IConsoleManager.h"
+#include "CanvasItem.h"
+#include "CanvasTypes.h"
+
+#if SUPPORT_TOUCH_INPUT_DISPLAY
+
+
+static int32 GInputEnableTouchDebugVisualizer = 0;
+static FAutoConsoleVariableRef CVarInputEnableTouchDebugVisualizer(
+	TEXT("Input.Debug.ShowTouches"),
+	GInputEnableTouchDebugVisualizer,
+	TEXT("Whether to show touch input on screen.")
+);
+
+
+
+FTouchInputVisualizer::FTouchInputVisualizer()
+{
+	if (!IsRunningCommandlet() && !IsRunningDedicatedServer())
+	{
+		// install input hook after the engine is initialized
+		FCoreDelegates::OnPostEngineInit.AddLambda([this]()
+		{
+			FSlateApplication::Get().RegisterInputPreProcessor(AsShared());
+		});
+
+		// remove input hook when the engine is shutting down
+		FCoreDelegates::OnEnginePreExit.AddLambda([this]()
+		{
+			FSlateApplication::Get().UnregisterInputPreProcessor(AsShared());
+		});
+	}
+}
+
+
+bool FTouchInputVisualizer::HandleMouseMoveEvent(FSlateApplication& SlateApp, const FPointerEvent& MouseEvent)
+{
+	if (MouseEvent.IsTouchEvent())
+	{
+		DebugTouchPoints.Add(MouseEvent.GetPointerIndex(), FDebugTouchPoint( MouseEvent.GetScreenSpacePosition(), MouseEvent.GetTouchForce() ));
+	}
+
+	return false;
+}
+
+bool FTouchInputVisualizer::HandleMouseButtonDownEvent(FSlateApplication& SlateApp, const FPointerEvent& MouseEvent)
+{
+	if (MouseEvent.IsTouchEvent())
+	{
+		DebugTouchPoints.Add( MouseEvent.GetPointerIndex(), FDebugTouchPoint( MouseEvent.GetScreenSpacePosition(), MouseEvent.GetTouchForce() ));
+	}
+
+	return false;
+}
+
+bool FTouchInputVisualizer::HandleMouseButtonUpEvent(FSlateApplication& SlateApp, const FPointerEvent& MouseEvent)
+{
+	if (MouseEvent.IsTouchEvent())
+	{
+		DebugTouchPoints.Remove( MouseEvent.GetPointerIndex() );
+	}
+
+	return false;
+}
+
+
+void FTouchInputVisualizer::OnDebugDraw(UCanvas* Canvas)
+{
+	if (GInputEnableTouchDebugVisualizer)
+	{
+		// adjust canvas transform (it defaults to a 1080p canvas but the touch points are in absolute screen positions)
+		const FVector Scale(Canvas->ClipX / Canvas->SizeX, Canvas->ClipY / Canvas->SizeY, 1.0);
+		Canvas->Canvas->PushRelativeTransform(FScaleMatrix(Scale));
+
+		// draw a circle for each active touch point (NB. can't use DrawDebugCanvas2DCircle because it's not available in Test)
+		static const FColor Color = FColor::Yellow;
+		static const int32 NumSides = 20;
+		static const float LineThickness = 3.0f;
+		static const float AngleDelta = TWO_PI / NumSides;
+
+		const float MaxRadius = (Canvas->SizeX * 0.02);
+		for (auto& Itr : DebugTouchPoints)
+		{
+			const FDebugTouchPoint& TouchPoint = Itr.Value;
+
+			const float Radius = MaxRadius * (1.0f + TouchPoint.Pressure);
+
+			FVector2D LastVertex = TouchPoint.Center + (FVector2D::UnitX() * Radius);
+			for (int32 SideIndex = 0; SideIndex < NumSides; SideIndex++)
+			{
+				const FVector2D Vertex = TouchPoint.Center + (FVector2D::UnitX() * FMath::Cos(AngleDelta * (SideIndex + 1)) + FVector2D::UnitY() * FMath::Sin(AngleDelta * (SideIndex + 1))) * Radius;
+				FCanvasLineItem LineItem(LastVertex, Vertex);
+				LineItem.LineThickness = LineThickness;
+				LineItem.SetColor(Color);
+				Canvas->DrawItem(LineItem);
+				LastVertex = Vertex;
+			}
+		}
+
+		// restore original canvas transform
+		Canvas->Canvas->PopTransform();
+	}
+}
+
+
+#endif //SUPPORT_TOUCH_INPUT_DISPLAY
diff --git a/Engine/Plugins/Runtime/InputDebugging/Source/InputDebugging/Private/TouchInputVisualizer.h b/Engine/Plugins/Runtime/InputDebugging/Source/InputDebugging/Private/TouchInputVisualizer.h
new file mode 100644
--- /dev/null
+++ b/Engine/Plugins/Runtime/InputDebugging/Source/InputDebugging/Private/TouchInputVisualizer.h
@@ -0,0 +1,47 @@
+// Copyright Epic Games, Inc. All Rights Reserved.
+
+#pragma once
+
+#include "Misc/Build.h"
+
+#if !defined(SUPPORT_TOUCH_INPUT_DISPLAY)
+#define SUPPORT_TOUCH_INPUT_DISPLAY !UE_BUILD_SHIPPING
+#endif
+
+#if SUPPORT_TOUCH_INPUT_DISPLAY
+#include "Framework/Application/IInputProcessor.h"
+#include "Delegates/IDelegateInstance.h"
+#include "Containers/Map.h"
+#include "Templates/SharedPointer.h"
+
+class FTouchInputVisualizer : public IInputProcessor, public TSharedFromThis<FTouchInputVisualizer>
+{
+public:
+	FTouchInputVisualizer();
+	virtual ~FTouchInputVisualizer() = default;
+
+	void OnDebugDraw(class UCanvas* Canvas);
+
+protected:
+	// IInputProcessor interface
+	virtual void Tick(const float DeltaTime, FSlateApplication& SlateApp, TSharedRef<ICursor> Cursor) override {};
+	virtual bool HandleMouseMoveEvent(FSlateApplication& SlateApp, const FPointerEvent& MouseEvent) override;
+	virtual bool HandleMouseButtonDownEvent(FSlateApplication& SlateApp, const FPointerEvent& MouseEvent) override;
+	virtual bool HandleMouseButtonUpEvent(FSlateApplication& SlateApp, const FPointerEvent& MouseEvent) override;
+
+	struct FDebugTouchPoint
+	{
+		explicit FDebugTouchPoint( const FVector2D& Center, float Pressure)
+		{
+			this->Center = Center;
+			this->Pressure = Pressure;
+		}
+
+		FVector2D Center;
+		float Pressure;
+	};
+
+	TMap<uint32, FDebugTouchPoint> DebugTouchPoints;
+
+};
+#endif //SUPPORT_TOUCH_INPUT_DISPLAY
diff --git a/Engine/Source/Runtime/Engine/Public/ShowFlagsValues.inl b/Engine/Source/Runtime/Engine/Public/ShowFlagsValues.inl
--- a/Engine/Source/Runtime/Engine/Public/ShowFlagsValues.inl
+++ b/Engine/Source/Runtime/Engine/Public/ShowFlagsValues.inl
@@ -404,8 +404,8 @@ SHOWFLAG_FIXED_IN_SHIPPING(0, DrawOnlyVSMInvalidatingGeo, SFG_Visualize, NSLOCTE
 /** Debug the Virtual Texture System */
 SHOWFLAG_FIXED_IN_SHIPPING(0, VirtualTextureResidency, SFG_Hidden, NSLOCTEXT("UnrealEd", "VirtualTextureResidency", "Virtual Texture Residency"))
 
-/** Debug DVR Streaming System System */
-SHOWFLAG_FIXED_IN_SHIPPING(0, DVRStreamingSystem, SFG_Hidden, NSLOCTEXT("UnrealEd", "DVRStreamingSystem", "DVR Streaming System"))
+/** Input Debug Visualizer */
+SHOWFLAG_FIXED_IN_SHIPPING(0, InputDebugVisualizer, SFG_Hidden, NSLOCTEXT("UnrealEd", "InputDebugVisualizer", "Input Debug Visualizer"))
 
 /** Use screen space tracing in Lumen */
 SHOWFLAG_ALWAYS_ACCESSIBLE(LumenScreenTraces, SFG_Lumen, NSLOCTEXT("UnrealEd", "LumenScreenTracesSF", "Screen Traces"))
diff --git a/Engine/Source/Runtime/InputCore/Public/IInputDebuggingInterface.h b/Engine/Source/Runtime/InputCore/Public/IInputDebuggingInterface.h
new file mode 100644
--- /dev/null
+++ b/Engine/Source/Runtime/InputCore/Public/IInputDebuggingInterface.h
@@ -0,0 +1,31 @@
+// Copyright Epic Games, Inc. All Rights Reserved.
+
+#pragma once
+
+#include "CoreMinimal.h"
+#include "Features/IModularFeatures.h"
+#include "Features/IModularFeature.h"
+#include "Modules/ModuleInterface.h"
+
+/**
+ * The public interface of the InputDeviceModule
+ */
+class IInputDebuggingInterface : public IModuleInterface, public IModularFeature
+{
+public:
+	static FName GetModularFeatureName()
+	{
+		static FName FeatureName = FName(TEXT("InputDebuggingInterface"));
+		return FeatureName;
+	}
+
+	static inline bool IsAvailable()
+	{
+		return IModularFeatures::Get().IsModularFeatureAvailable(GetModularFeatureName());
+	}
+
+	static inline IInputDebuggingInterface& Get()
+	{
+		return IModularFeatures::Get().GetModularFeature<IInputDebuggingInterface>(GetModularFeatureName());
+	}
+};
