Add missing module dependency.

#rnx

diff --git a/Engine/Plugins/Runtime/InputDebugging/Source/InputDebugging/InputDebugging.Build.cs b/Engine/Plugins/Runtime/InputDebugging/Source/InputDebugging/InputDebugging.Build.cs
--- a/Engine/Plugins/Runtime/InputDebugging/Source/InputDebugging/InputDebugging.Build.cs
+++ b/Engine/Plugins/Runtime/InputDebugging/Source/InputDebugging/InputDebugging.Build.cs
@@ -10,6 +10,7 @@ namespace UnrealBuildTool.Rules
 				new string[] {
 					"Core",
 					"Engine",
+					"SlateCore",
 					"Slate",
 					"InputCore",
 				}
