Fix nopch issues
#rnx

diff --git a/Engine/Plugins/Media/ImgMedia/Source/ExrReaderGpu/Private/ExrReaderGpu.cpp b/Engine/Plugins/Media/ImgMedia/Source/ExrReaderGpu/Private/ExrReaderGpu.cpp
--- a/Engine/Plugins/Media/ImgMedia/Source/ExrReaderGpu/Private/ExrReaderGpu.cpp
+++ b/Engine/Plugins/Media/ImgMedia/Source/ExrReaderGpu/Private/ExrReaderGpu.cpp
@@ -4,6 +4,7 @@
 #if PLATFORM_WINDOWS
 
 #include "ExrReaderGpuModule.h"
+#include "HAL/IConsoleManager.h"
 
 PRAGMA_DEFAULT_VISIBILITY_START
 THIRD_PARTY_INCLUDES_START
diff --git a/Engine/Plugins/Runtime/InputDebugging/Source/InputDebugging/Private/InputDebuggingModule.cpp b/Engine/Plugins/Runtime/InputDebugging/Source/InputDebugging/Private/InputDebuggingModule.cpp
--- a/Engine/Plugins/Runtime/InputDebugging/Source/InputDebugging/Private/InputDebuggingModule.cpp
+++ b/Engine/Plugins/Runtime/InputDebugging/Source/InputDebugging/Private/InputDebuggingModule.cpp
@@ -4,7 +4,7 @@
 #include "Debug/DebugDrawService.h"
 #include "Engine/Canvas.h"
 #include "TouchInputVisualizer.h"
-
+#include "Modules/ModuleManager.h"
 
 class FInputDebuggingModule : public IInputDebuggingInterface
 {
diff --git a/Engine/Source/Editor/AnimationEditMode/Public/AnimationEditContext.h b/Engine/Source/Editor/AnimationEditMode/Public/AnimationEditContext.h
--- a/Engine/Source/Editor/AnimationEditMode/Public/AnimationEditContext.h
+++ b/Engine/Source/Editor/AnimationEditMode/Public/AnimationEditContext.h
@@ -2,6 +2,7 @@
 
 #pragma once
 
+#include "UObject/Interface.h"
 #include "UObject/UObjectHierarchyFwd.h"
 
 #include "AnimationEditContext.generated.h"
