add missing include

#rnx

diff --git a/Engine/Plugins/Runtime/InputDebugging/Source/InputDebugging/Private/TouchInputVisualizer.cpp b/Engine/Plugins/Runtime/InputDebugging/Source/InputDebugging/Private/TouchInputVisualizer.cpp
--- a/Engine/Plugins/Runtime/InputDebugging/Source/InputDebugging/Private/TouchInputVisualizer.cpp
+++ b/Engine/Plugins/Runtime/InputDebugging/Source/InputDebugging/Private/TouchInputVisualizer.cpp
@@ -5,6 +5,7 @@
 #include "Engine/Canvas.h"
 #include "Misc/CoreDelegates.h"
 #include "HAL/IConsoleManager.h"
+#include "Framework/Application/SlateApplication.h"
 #include "CanvasItem.h"
 #include "CanvasTypes.h"
 
