disable input debugging plugin in Shipping

#rnx

diff --git a/Engine/Plugins/Runtime/InputDebugging/InputDebugging.uplugin b/Engine/Plugins/Runtime/InputDebugging/InputDebugging.uplugin
--- a/Engine/Plugins/Runtime/InputDebugging/InputDebugging.uplugin
+++ b/Engine/Plugins/Runtime/InputDebugging/InputDebugging.uplugin
@@ -18,7 +18,8 @@
 		{
 			"Name": "InputDebugging",
 			"Type": "Runtime",
-			"LoadingPhase": "Default"
+			"LoadingPhase": "Default",
+			"TargetConfigurationDenyList": [ "Shipping" ]
 		}
 	]
 }
\ No newline at end of file
