Add an Input Debugging Editor module. This is where any generic Input Debugging tools can live now and in the future.

Add in a slate notification when a gamepad's state changes in the editor, which displays the Platform User, Device ID, and connection state. In the future I'd like to make a small window that shows you all currently connected devices and the data associated with them, as well possiblly as visualize some of their input or allow you to inject input to different pads.

diff --git a/Engine/Plugins/Runtime/InputDebugging/InputDebugging.uplugin b/Engine/Plugins/Runtime/InputDebugging/InputDebugging.uplugin
--- a/Engine/Plugins/Runtime/InputDebugging/InputDebugging.uplugin
+++ b/Engine/Plugins/Runtime/InputDebugging/InputDebugging.uplugin
@@ -20,6 +20,11 @@
 			"Type": "Runtime",
 			"LoadingPhase": "Default",
 			"TargetConfigurationDenyList": [ "Shipping" ]
+		},
+		{
+			"Name": "InputDebuggingEditor",
+			"Type": "Editor",
+			"LoadingPhase": "Default"
 		}
 	]
 }
\ No newline at end of file
diff --git a/Engine/Plugins/Runtime/InputDebugging/Source/InputDebuggingEditor/InputDebuggingEditor.Build.cs b/Engine/Plugins/Runtime/InputDebugging/Source/InputDebuggingEditor/InputDebuggingEditor.Build.cs
new file mode 100644
--- /dev/null
+++ b/Engine/Plugins/Runtime/InputDebugging/Source/InputDebuggingEditor/InputDebuggingEditor.Build.cs
@@ -0,0 +1,20 @@
+// Copyright Epic Games, Inc. All Rights Reserved.
+
+namespace UnrealBuildTool.Rules
+{
+	public class InputDebuggingEditor : ModuleRules
+	{
+		public InputDebuggingEditor(ReadOnlyTargetRules Target) : base(Target)
+		{
+			PrivateDependencyModuleNames.AddRange(
+				new string[] {
+					"Core",
+					"ApplicationCore",
+					"Slate",
+					"SlateCore",
+					"MainFrame",
+				}
+			);
+		}
+	}
+}
\ No newline at end of file
diff --git a/Engine/Plugins/Runtime/InputDebugging/Source/InputDebuggingEditor/Private/InputDebuggingEditorModule.cpp b/Engine/Plugins/Runtime/InputDebugging/Source/InputDebuggingEditor/Private/InputDebuggingEditorModule.cpp
new file mode 100644
--- /dev/null
+++ b/Engine/Plugins/Runtime/InputDebugging/Source/InputDebuggingEditor/Private/InputDebuggingEditorModule.cpp
@@ -0,0 +1,139 @@
+// Copyright Epic Games, Inc. All Rights Reserved.
+
+#include "InputDebuggingEditorModule.h"
+#include "GenericPlatform/GenericPlatformInputDeviceMapper.h"
+#include "Framework/Notifications/NotificationManager.h"
+#include "Widgets/Notifications/SNotificationList.h"
+#include "Interfaces/IMainFrameModule.h"
+#include "Styling/SlateStyle.h"
+#include "Styling/SlateStyleMacros.h"
+#include "Styling/SlateStyleRegistry.h"
+#include "Misc/ConfigCacheIni.h"
+
+#define LOCTEXT_NAMESPACE "InputDebuggingEditor"
+
+/** Custom style set for the input debugger */
+class FInputDebuggerSlateStyle final : public FSlateStyleSet
+{
+public:
+	FInputDebuggerSlateStyle()
+		: FSlateStyleSet("InputDebugger")
+	{
+		SetParentStyleName(FAppStyle::GetAppStyleSetName());
+		
+		SetContentRoot(FPaths::EngineContentDir() / TEXT("Editor/Slate"));
+		SetCoreContentRoot(FPaths::EngineContentDir() / TEXT("Slate"));
+
+		// Input Debugger icon for gamepad
+		static const FVector2D Icon40x40 = FVector2D(40.0f, 40.0f);
+
+		Set("InputDebugger.Gamepad", new IMAGE_BRUSH_SVG("Starship/AssetIcons/PlayerController_16", Icon40x40));
+	}
+};
+
+void FInputDebuggingEditorModule::StartupModule()
+{
+	// Listen for when input devices change
+	IPlatformInputDeviceMapper::Get().GetOnInputDeviceConnectionChange().AddRaw(this, &FInputDebuggingEditorModule::OnInputDeviceConnectionChange);
+	
+	// Make a new style set for an input debugger, which will register any custom icons for the types in this module
+	StyleSet = MakeShared<FInputDebuggerSlateStyle>();
+	FSlateStyleRegistry::RegisterSlateStyle(*StyleSet.Get());
+}
+
+void FInputDebuggingEditorModule::ShutdownModule()
+{
+	// Remove any input listeners
+	IPlatformInputDeviceMapper::Get().GetOnInputDeviceConnectionChange().RemoveAll(this);
+	
+	// Unregister slate stylings
+	if (StyleSet.IsValid())
+	{
+		FSlateStyleRegistry::UnRegisterSlateStyle(*StyleSet.Get());
+	}
+}
+
+const FText& FInputDebuggingEditorModule::ConnectionStateToText(const EInputDeviceConnectionState NewConnectionState)
+{
+	// "An input device..."
+	static const FText InvalidLabel = LOCTEXT("InvalidStateLabel", "is in an invalid state");
+	static const FText DisconnectedLabel = LOCTEXT("DisconnectedStateLabel", "has been disconnected");
+	static const FText ConnectedLabel = LOCTEXT("ConnectedStateLabel", "has been connected.");
+	static const FText UnknownLabel = LOCTEXT("UnknownStateLabel", "has an unknown state!");
+	
+	switch (NewConnectionState)
+	{
+		case EInputDeviceConnectionState::Invalid:
+			return InvalidLabel;
+		break;
+
+		case EInputDeviceConnectionState::Disconnected:
+			return DisconnectedLabel;
+		break;
+		
+		case EInputDeviceConnectionState::Connected:
+			return ConnectedLabel;
+		break;
+		
+		case EInputDeviceConnectionState::Unknown:
+		default:
+			return UnknownLabel;
+		break;
+	}
+}
+
+void FInputDebuggingEditorModule::OnInputDeviceConnectionChange(EInputDeviceConnectionState NewConnectionState, FPlatformUserId PlatformUserId, FInputDeviceId InputDeviceId)
+{
+	/** Utility functions for notifications */
+	struct Local
+	{
+		static ECheckBoxState GetDontAskAgainCheckBoxState()
+		{
+			// Check the config for any preferences
+			bool bSuppressNotification = false;
+			GConfig->GetBool(TEXT("InputDebuggingEditor"), TEXT("SuppressInputDeviceConnectionChangeNotification"), bSuppressNotification, GEditorPerProjectIni);
+			return bSuppressNotification ? ECheckBoxState::Checked : ECheckBoxState::Unchecked;
+		}
+
+		static void OnDontAskAgainCheckBoxStateChanged(ECheckBoxState NewState)
+		{
+			// If the user selects to not show this again, set that in the config so we know about it in between sessions
+			const bool bSuppressNotification = (NewState == ECheckBoxState::Checked);
+			GConfig->SetBool(TEXT("InputDebuggingEditor"), TEXT("SuppressInputDeviceConnectionChangeNotification"), bSuppressNotification, GEditorPerProjectIni);
+		}
+	}; 
+	
+	// Send a slate notification that there was a new gamepad plugged in
+	FFormatNamedArguments MainMessageArgs;
+	MainMessageArgs.Add(TEXT("StateMessage"), ConnectionStateToText(NewConnectionState));
+	const FText Message = FText::Format(LOCTEXT("InputDeviceConnectionChangeMessage", "An input device {StateMessage}"), MainMessageArgs);
+
+	// Add details about what DeviceID and PlatformUserId it has been mapped to
+	FFormatNamedArguments SubMessageArgs;
+	SubMessageArgs.Add(TEXT("PlatformUser"), PlatformUserId.GetInternalId());
+	SubMessageArgs.Add(TEXT("DeviceId"), InputDeviceId.GetId());
+	const FText SubMessage = FText::Format(LOCTEXT("InputDeviceConnectionChangeSubMessage", "Platform User ID: {PlatformUser}\nDevice ID: {DeviceId}"), SubMessageArgs);
+
+	// Send out a slate notification that there was an input device connection change
+	FNotificationInfo* Info = new FNotificationInfo(Message);
+	Info->SubText = SubMessage;
+	Info->ExpireDuration = 5.0f;
+	Info->bFireAndForget = true;
+
+	// Set the icon to a nifty gamepad
+	if (StyleSet)
+	{
+		Info->Image = StyleSet->GetBrush(TEXT("InputDebugger.Gamepad"));
+	}
+
+	// Add a "Don't show this again" option
+	Info->CheckBoxState = TAttribute<ECheckBoxState>::Create(&Local::GetDontAskAgainCheckBoxState);
+	Info->CheckBoxStateChanged = FOnCheckStateChanged::CreateStatic(&Local::OnDontAskAgainCheckBoxStateChanged);
+	Info->CheckBoxText = NSLOCTEXT("ModalDialogs", "DefaultCheckBoxMessage", "Don't show this again");
+	
+	FSlateNotificationManager::Get().QueueNotification(Info);
+}
+
+IMPLEMENT_MODULE(FInputDebuggingEditorModule, InputDebuggingEditor);
+
+#undef LOCTEXT_NAMESPACE
\ No newline at end of file
diff --git a/Engine/Plugins/Runtime/InputDebugging/Source/InputDebuggingEditor/Public/InputDebuggingEditorModule.h b/Engine/Plugins/Runtime/InputDebugging/Source/InputDebuggingEditor/Public/InputDebuggingEditorModule.h
new file mode 100644
--- /dev/null
+++ b/Engine/Plugins/Runtime/InputDebugging/Source/InputDebuggingEditor/Public/InputDebuggingEditorModule.h
@@ -0,0 +1,29 @@
+// Copyright Epic Games, Inc. All Rights Reserved.
+
+#pragma once
+
+#include "CoreMinimal.h"
+#include "Modules/ModuleManager.h"
+
+class FSlateStyleSet;
+
+class FInputDebuggingEditorModule : public IModuleInterface
+{
+public:
+	/** IModuleInterface implementation */
+	virtual void StartupModule() override;
+	virtual void ShutdownModule() override;
+
+private:
+
+	/**
+	 * Returns an FText label for the given input device connection state
+	 */
+	static const FText& ConnectionStateToText(const EInputDeviceConnectionState NewConnectionState);
+
+	/** Callback for when an input device has a connection state change (New gamepad plugged in, unplugged, etc) */
+	void OnInputDeviceConnectionChange(EInputDeviceConnectionState NewConnectionState, FPlatformUserId PlatformUserId, FInputDeviceId InputDeviceId);
+
+	/** A custom slate style for input debugging related things */
+	TSharedPtr<FSlateStyleSet> StyleSet;
+};
\ No newline at end of file
