Fix the "Dont show again" checkbox not working when an input device connection changes

#rnx

diff --git a/Engine/Plugins/Runtime/InputDebugging/Source/InputDebuggingEditor/Private/InputDebuggingEditorModule.cpp b/Engine/Plugins/Runtime/InputDebugging/Source/InputDebuggingEditor/Private/InputDebuggingEditorModule.cpp
--- a/Engine/Plugins/Runtime/InputDebugging/Source/InputDebuggingEditor/Private/InputDebuggingEditorModule.cpp
+++ b/Engine/Plugins/Runtime/InputDebugging/Source/InputDebuggingEditor/Private/InputDebuggingEditorModule.cpp
@@ -87,12 +87,18 @@ void FInputDebuggingEditorModule::OnInputDeviceConnectionChange(EInputDeviceConn
 	/** Utility functions for notifications */
 	struct Local
 	{
-		static ECheckBoxState GetDontAskAgainCheckBoxState()
+
+		static bool ShouldSuppressModal()
 		{
-			// Check the config for any preferences
 			bool bSuppressNotification = false;
 			GConfig->GetBool(TEXT("InputDebuggingEditor"), TEXT("SuppressInputDeviceConnectionChangeNotification"), bSuppressNotification, GEditorPerProjectIni);
-			return bSuppressNotification ? ECheckBoxState::Checked : ECheckBoxState::Unchecked;
+			return bSuppressNotification;
+		}
+
+		static ECheckBoxState GetDontAskAgainCheckBoxState()
+		{
+			// Check the config for any preferences			
+			return ShouldSuppressModal() ? ECheckBoxState::Checked : ECheckBoxState::Unchecked;
 		}
 
 		static void OnDontAskAgainCheckBoxStateChanged(ECheckBoxState NewState)
@@ -102,6 +108,12 @@ void FInputDebuggingEditorModule::OnInputDeviceConnectionChange(EInputDeviceConn
 			GConfig->SetBool(TEXT("InputDebuggingEditor"), TEXT("SuppressInputDeviceConnectionChangeNotification"), bSuppressNotification, GEditorPerProjectIni);
 		}
 	}; 
+
+	// If the user has specified to supress this pop up, then just early out and exit	
+	if (Local::ShouldSuppressModal())
+	{
+		return;
+	}
 	
 	// Send a slate notification that there was a new gamepad plugged in
 	FFormatNamedArguments MainMessageArgs;
