Add debug drawing and commands for Input Device Properties. Move "showdebug devices" out of enhanced input, and into the input debugging plugin because it doesn't actually require Enhanced Input.

#rnx

diff --git a/Engine/Plugins/EnhancedInput/Config/Input.ini b/Engine/Plugins/EnhancedInput/Config/Input.ini
--- a/Engine/Plugins/EnhancedInput/Config/Input.ini
+++ b/Engine/Plugins/EnhancedInput/Config/Input.ini
@@ -1,3 +1,2 @@
 [/Script/EngineSettings.ConsoleSettings]
 +ManualAutoCompleteList=(Command="ShowDebug EnhancedInput",Desc="Displays debug information about the current state of any Enhanced Input Mapping Contexts")
-+ManualAutoCompleteList=(Command="ShowDebug Devices",Desc="Displays all current Platform Users and their Input Devices.")
diff --git a/Engine/Plugins/EnhancedInput/Source/EnhancedInput/Private/EnhancedInputModule.cpp b/Engine/Plugins/EnhancedInput/Source/EnhancedInput/Private/EnhancedInputModule.cpp
--- a/Engine/Plugins/EnhancedInput/Source/EnhancedInput/Private/EnhancedInputModule.cpp
+++ b/Engine/Plugins/EnhancedInput/Source/EnhancedInput/Private/EnhancedInputModule.cpp
@@ -306,7 +306,6 @@ void FEnhancedInputModule::Tick(float DeltaTime)
 void FEnhancedInputModule::OnShowDebugInfo(AHUD* HUD, UCanvas* Canvas, const FDebugDisplayInfo& DisplayInfo, float& YL, float& YPos)
 {
 	static const FName NAME_EnhancedInput("EnhancedInput");
-	static const FName NAME_PlatformDevices("Devices");
 	if (Canvas)
 	{
 		if (HUD->ShouldDisplayDebug(NAME_EnhancedInput))
@@ -325,14 +324,6 @@ void FEnhancedInputModule::OnShowDebugInfo(AHUD* HUD, UCanvas* Canvas, const FDe
 				FirstPlayer->ShowDebugInfo(Canvas);
 			}
 		}
-		
-		if (HUD->ShouldDisplayDebug(NAME_PlatformDevices))
-		{
-			UEnhancedInputLibrary::ForEachSubsystem([Canvas](IEnhancedInputSubsystemInterface* Subsystem)
-			{
-				Subsystem->ShowPlatformInputDebugInfo(Canvas);
-			});
-		}
 	}
 	else
 	{
diff --git a/Engine/Plugins/EnhancedInput/Source/EnhancedInput/Private/EnhancedInputSubsystemInterfaceDebug.cpp b/Engine/Plugins/EnhancedInput/Source/EnhancedInput/Private/EnhancedInputSubsystemInterfaceDebug.cpp
--- a/Engine/Plugins/EnhancedInput/Source/EnhancedInput/Private/EnhancedInputSubsystemInterfaceDebug.cpp
+++ b/Engine/Plugins/EnhancedInput/Source/EnhancedInput/Private/EnhancedInputSubsystemInterfaceDebug.cpp
@@ -14,7 +14,6 @@
 #include "InputTriggers.h"
 #include "ImageUtils.h"
 #include "EnhancedInputPlatformSettings.h"
-#include "Framework/Application/SlateApplication.h"
 #include "GenericPlatform/GenericPlatformInputDeviceMapper.h"
 #include "UObject/Package.h"
 
@@ -274,93 +273,6 @@ void IEnhancedInputSubsystemInterface::ShowDebugInfo(UCanvas* Canvas)
 	}
 }
 
-void IEnhancedInputSubsystemInterface::ShowPlatformInputDebugInfo(class UCanvas* Canvas)
-{
-	if (!Canvas)
-	{
-		return;
-	}
-	FDisplayDebugManager& DisplayDebugManager = Canvas->DisplayDebugManager;
-
-	// Show title
-	DisplayDebugManager.SetFont(GEngine->GetMediumFont());
-	DisplayDebugManager.SetDrawColor(FColor::Orange);
-	DisplayDebugManager.DrawString(TEXT("\n\nPlatform Input Devices"));
-
-	// Gather all current platform users
-	IPlatformInputDeviceMapper& DeviceMapper = IPlatformInputDeviceMapper::Get();
-	TArray<FPlatformUserId> PlatformUsers;
-	const int32 NumUsers = DeviceMapper.GetAllActiveUsers(PlatformUsers);
-	
-	DisplayDebugManager.SetDrawColor(FColor::White);
-	DisplayDebugManager.DrawString(FString::Printf(TEXT("Total Platform Users: %d\n"), NumUsers));
-
-	// Show all the connected input devices for each available platform user
-	// TODO: We could make two commands, one to display all platform users and one to only show this local player's data
-	TArray<FInputDeviceId> UserDeviceIds;
-	for (FPlatformUserId& PlatUser : PlatformUsers)
-	{
-		const int32 NumDevices = DeviceMapper.GetAllInputDevicesForUser(PlatUser, UserDeviceIds);
-		FString UserDebugInfo = FString::Printf(TEXT("Platform User ID '%d' has '%d' input devices"), PlatUser.GetInternalId(), NumDevices);
-		// Show in red if there are no devices, that is a problem!
-		DisplayDebugManager.SetDrawColor(NumDevices > 0 ? FColor::White : FColor::Red);
-		DisplayDebugManager.DrawString(UserDebugInfo);
-
-		// Display the debug info about each input device
-		for (FInputDeviceId& Device : UserDeviceIds)
-		{
-			static const float InputDeviceXOffset = 10.0f;
-			// Display the input device ID
-			FString DeviceIdString = FString::Printf(TEXT("Device ID: %d"), Device.GetId());
-			FString SlateUserIdString = TEXT("Slate UserID is unknown (Application is not initialized)");
-			if (FSlateApplication::IsInitialized())
-			{
-				TOptional<int32> SlateUserId = FSlateApplication::Get().GetUserIndexForInputDevice(Device);
-				if (SlateUserId.IsSet())
-				{
-					SlateUserIdString = FString::Printf(TEXT("Slate UserId ID: %d"), SlateUserId.GetValue());
-				}
-				else
-				{
-					SlateUserIdString = TEXT("Slate UserId ID: invalid");
-				}
-			}
-			
-			DisplayDebugManager.SetDrawColor(FColor::White);
-			DisplayDebugManager.DrawString(DeviceIdString, InputDeviceXOffset);
-			DisplayDebugManager.DrawString(SlateUserIdString, InputDeviceXOffset);
-			
-			// Display the connection state
-			FString ConnectionStateString;
-			FColor ConnectionStateColor = FColor::White;
-			const EInputDeviceConnectionState DeviceState = DeviceMapper.GetInputDeviceConnectionState(Device);
-			
-			switch(DeviceState)
-			{
-				case EInputDeviceConnectionState::Invalid:
-					ConnectionStateString = TEXT("INVALID");
-					ConnectionStateColor = FColor::Red;
-					break;
-				case EInputDeviceConnectionState::Unknown:
-					ConnectionStateString = TEXT("Unknown");
-					ConnectionStateColor = FColor::Orange;
-					break;
-				case EInputDeviceConnectionState::Disconnected:
-					ConnectionStateString = TEXT("Disconnected");
-					ConnectionStateColor = FColor::Silver;
-					break;
-				case EInputDeviceConnectionState::Connected:
-					ConnectionStateString = TEXT("Connected");
-					ConnectionStateColor = FColor::Green;
-					break;
-			}
-			
-			DisplayDebugManager.SetDrawColor(ConnectionStateColor);
-			DisplayDebugManager.DrawString(ConnectionStateString, InputDeviceXOffset);
-		}
-	}
-}
-
 void IEnhancedInputSubsystemInterface::ShowDebugActionModifiers(UCanvas* Canvas, const UInputAction* Action)
 {
 	FDisplayDebugManager& DisplayDebugManager = Canvas->DisplayDebugManager;
diff --git a/Engine/Plugins/EnhancedInput/Source/EnhancedInput/Public/EnhancedInputSubsystemInterface.h b/Engine/Plugins/EnhancedInput/Source/EnhancedInput/Public/EnhancedInputSubsystemInterface.h
--- a/Engine/Plugins/EnhancedInput/Source/EnhancedInput/Public/EnhancedInputSubsystemInterface.h
+++ b/Engine/Plugins/EnhancedInput/Source/EnhancedInput/Public/EnhancedInputSubsystemInterface.h
@@ -309,11 +309,6 @@ private:
 	// Debug visualization implemented in EnhancedInputSubsystemsDebug.cpp
 	void ShowDebugInfo(class UCanvas* Canvas);
 	
-	// Debug visualization of any platform input devices (IPlatformInputDeviceMapper)
-	// TODO: We should get a better place to put this instead of the Enhanced Input plugin.
-	// It can't go in InputCore because that doesn't depend on Engine for debug drawing
-	void ShowPlatformInputDebugInfo(class UCanvas* Canvas);
-	
 	void ShowDebugActionModifiers(UCanvas* Canvas, const UInputAction* Action);
 	static void PurgeDebugVisualizations();
 };
\ No newline at end of file
diff --git a/Engine/Plugins/Runtime/InputDebugging/Config/Input.ini b/Engine/Plugins/Runtime/InputDebugging/Config/Input.ini
new file mode 100644
--- /dev/null
+++ b/Engine/Plugins/Runtime/InputDebugging/Config/Input.ini
@@ -0,0 +1,3 @@
+[/Script/EngineSettings.ConsoleSettings]
++ManualAutoCompleteList=(Command="ShowDebug DeviceProperty",Desc="Display debug information about currently active Input Device Properties")
++ManualAutoCompleteList=(Command="ShowDebug Devices",Desc="Displays all current Platform Users and their Input Devices.")
\ No newline at end of file
diff --git a/Engine/Plugins/Runtime/InputDebugging/Source/InputDebugging/InputDebugging.Build.cs b/Engine/Plugins/Runtime/InputDebugging/Source/InputDebugging/InputDebugging.Build.cs
--- a/Engine/Plugins/Runtime/InputDebugging/Source/InputDebugging/InputDebugging.Build.cs
+++ b/Engine/Plugins/Runtime/InputDebugging/Source/InputDebugging/InputDebugging.Build.cs
@@ -13,6 +13,8 @@ namespace UnrealBuildTool.Rules
 					"SlateCore",
 					"Slate",
 					"InputCore",
+					"CoreUObject",
+					"ApplicationCore"
 				}
 			);
 		}
diff --git a/Engine/Plugins/Runtime/InputDebugging/Source/InputDebugging/Private/InputDebuggingModule.cpp b/Engine/Plugins/Runtime/InputDebugging/Source/InputDebugging/Private/InputDebuggingModule.cpp
--- a/Engine/Plugins/Runtime/InputDebugging/Source/InputDebugging/Private/InputDebuggingModule.cpp
+++ b/Engine/Plugins/Runtime/InputDebugging/Source/InputDebugging/Private/InputDebuggingModule.cpp
@@ -4,6 +4,7 @@
 #include "Debug/DebugDrawService.h"
 #include "Engine/Canvas.h"
 #include "TouchInputVisualizer.h"
+#include "InputDeviceDebugTools.h"
 #include "Modules/ModuleManager.h"
 
 class FInputDebuggingModule : public IInputDebuggingInterface
@@ -23,6 +24,10 @@ public:
 #if SUPPORT_TOUCH_INPUT_DISPLAY
 		TouchInputVisualizer = MakeShared<FTouchInputVisualizer>();
 #endif//SUPPORT_TOUCH_INPUT_DISPLAY
+
+#if SUPPORT_INPUT_DEVICE_DEBUGGING
+		InputDeviceVisualizer = MakeShared<FInputDeviceDebugTools>();
+#endif	// SUPPORT_INPUT_DEVICE_DEBUGGING
 	}
 
 	virtual void ShutdownModule() override
@@ -32,6 +37,10 @@ public:
 		TouchInputVisualizer.Reset();
 #endif//SUPPORT_TOUCH_INPUT_DISPLAY
 
+#if SUPPORT_INPUT_DEVICE_DEBUGGING
+		InputDeviceVisualizer.Reset();
+#endif	// SUPPORT_INPUT_DEVICE_DEBUGGING
+
 		// unregister debug rendering callback
 		UDebugDrawService::Unregister(DrawOnCanvasDelegateHandle);
 
@@ -56,6 +65,9 @@ private:
 	TSharedPtr<FTouchInputVisualizer> TouchInputVisualizer;
 #endif//SUPPORT_TOUCH_INPUT_DISPLAY
 
+#if SUPPORT_INPUT_DEVICE_DEBUGGING
+	TSharedPtr<FInputDeviceDebugTools> InputDeviceVisualizer;
+#endif	// SUPPORT_INPUT_DEVICE_DEBUGGING
 };
 
-IMPLEMENT_MODULE( FInputDebuggingModule, InputDebugging);
+IMPLEMENT_MODULE(FInputDebuggingModule, InputDebugging);
diff --git a/Engine/Plugins/Runtime/InputDebugging/Source/InputDebugging/Private/InputDeviceDebugTools.cpp b/Engine/Plugins/Runtime/InputDebugging/Source/InputDebugging/Private/InputDeviceDebugTools.cpp
new file mode 100644
--- /dev/null
+++ b/Engine/Plugins/Runtime/InputDebugging/Source/InputDebugging/Private/InputDeviceDebugTools.cpp
@@ -0,0 +1,261 @@
+// Copyright Epic Games, Inc. All Rights Reserved.
+
+#include "InputDeviceDebugTools.h"
+
+#if SUPPORT_INPUT_DEVICE_DEBUGGING
+
+#include "GenericPlatform/GenericPlatformInputDeviceMapper.h"
+#include "Framework/Application/SlateApplication.h"
+#include "HAL/IConsoleManager.h"
+#include "GameFramework/InputDeviceSubsystem.h"
+#include "GameFramework/InputDeviceProperties.h"
+#include "GameFramework/InputSettings.h"
+#include "GameFramework/HUD.h"
+#include "Engine/Engine.h"
+#include "Engine/Canvas.h"
+#include "UObject/ObjectPtr.h"
+#include "Logging/LogMacros.h"
+#include "DrawDebugHelpers.h"
+
+FInputDeviceDebugTools::FInputDeviceDebugTools()
+{
+#if ENABLE_DRAW_DEBUG
+	AHUD::OnShowDebugInfo.AddStatic(&FInputDeviceDebugTools::OnShowDebugInfo);
+#endif
+	AddConsoleCommands();
+}
+
+FInputDeviceDebugTools::~FInputDeviceDebugTools()
+{
+	RemoveConsoleCommands();
+}
+
+void FInputDeviceDebugTools::AddConsoleCommands()
+{
+	ConsoleCommands.Add(IConsoleManager::Get().RegisterConsoleCommand(
+		TEXT("Input.ListAllHardwareDevices"),
+		TEXT("Log all the platform's currently available FHardwareDeviceIdentifier"),
+		FConsoleCommandWithWorldAndArgsDelegate::CreateRaw(this, &FInputDeviceDebugTools::ListAllKnownHardwareDeviceIdentifier),
+		ECVF_Cheat
+	));
+}
+
+void FInputDeviceDebugTools::RemoveConsoleCommands()
+{
+	// Unregister console commands
+	for (IConsoleCommand* Command : ConsoleCommands)
+	{
+		IConsoleManager::Get().UnregisterConsoleObject(Command);
+	}
+}
+
+void FInputDeviceDebugTools::OnShowDebugInfo(AHUD* HUD, UCanvas* Canvas, const FDebugDisplayInfo& DisplayInfo, float& YL, float& YPos)
+{
+	static const FName NAME_DeviceProperty("DeviceProperty");
+	static const FName NAME_HardwareDevice("Devices");
+
+	if (Canvas)
+	{
+		if (HUD->ShouldDisplayDebug(NAME_DeviceProperty))
+		{
+			OnShowDebugDeviceProperties(Canvas);
+		}
+
+		if (HUD->ShouldDisplayDebug(NAME_HardwareDevice))
+		{
+			OnShowDebugHardwareDevices(Canvas);
+		}
+	}
+}
+
+void FInputDeviceDebugTools::OnShowDebugDeviceProperties(UCanvas* Canvas)
+{
+	FDisplayDebugManager& DisplayDebugManager = Canvas->DisplayDebugManager;
+
+	// Show title
+	DisplayDebugManager.SetFont(GEngine->GetMediumFont());
+	DisplayDebugManager.SetDrawColor(FColor::Orange);
+	DisplayDebugManager.DrawString(TEXT("\n\nInput Device Properties"));
+
+	const UInputDeviceSubsystem* System = UInputDeviceSubsystem::Get();
+	if (!System)
+	{
+		DisplayDebugManager.SetDrawColor(FColor::Red);
+		DisplayDebugManager.DrawString(TEXT("\t\tUInputDeviceSubsystem is not initalized!"));
+		return;
+	}
+
+	if (System->ActiveProperties.IsEmpty())
+	{
+		DisplayDebugManager.SetDrawColor(FColor::White);
+		DisplayDebugManager.DrawString(TEXT("\t \t No active device properties"));
+		return;
+	}
+
+	// Sort the active properties by their owning Platform User
+	static TMap<FPlatformUserId, TArray<const FActiveDeviceProperty*>> PlatformUserDebugStrings;
+	{
+		PlatformUserDebugStrings.Empty();
+
+		for (const FActiveDeviceProperty& ActiveProperty : System->ActiveProperties)
+		{
+			PlatformUserDebugStrings.FindOrAdd(ActiveProperty.PlatformUser).Emplace(&ActiveProperty);
+		}
+	}	
+
+	// Actually display the debug information	
+	for (TPair<FPlatformUserId, TArray<const FActiveDeviceProperty*>> Pair : PlatformUserDebugStrings)
+	{
+		const FString PlatUserLabel = FString::Printf(TEXT("Platform User %d has %d active Device Properties"),
+			Pair.Key.GetInternalId(),
+			Pair.Value.Num());
+
+		DisplayDebugManager.SetDrawColor(FColor::Yellow);
+		DisplayDebugManager.DrawString(PlatUserLabel);
+
+		for (const FActiveDeviceProperty* ActiveProperty: Pair.Value)
+		{
+			// Draw the name of the property
+			DisplayDebugManager.SetDrawColor(FColor::White);
+			DisplayDebugManager.DrawString(ActiveProperty->Property->GetFName().ToString());
+
+			// Draw some information about the flags that have been set on the given property
+			DisplayDebugManager.SetDrawColor(FColor::Silver);
+			DisplayDebugManager.DrawString(FString::Printf(TEXT("\t Property Handle: %s"), *ActiveProperty->PropertyHandle.ToString()));
+			DisplayDebugManager.DrawString(FString::Printf(TEXT("\t Flags: \n \t \t \t bRemoveAfterEvaluationTime: %d \n \t \t \t bIgnoreTimeDilation: %d \n \t \t \t bPlayWhilePaused: %d"),
+				ActiveProperty->bRemoveAfterEvaluationTime, ActiveProperty->bIgnoreTimeDilation, ActiveProperty->bPlayWhilePaused));
+
+			// Highlight if a property's evaluation has completed. If can go higher then the duration if
+			// is set to not be removed after its evaluation is complete
+			if (ActiveProperty->EvaluatedDuration > ActiveProperty->Property->GetDuration())
+			{
+				DisplayDebugManager.SetDrawColor(FColor::Green);
+			}
+
+			DisplayDebugManager.DrawString(FString::Printf(TEXT("\t Eval Time: %f"), ActiveProperty->EvaluatedDuration));			
+
+			// TODO_BH: Pass in the canvas the the device property here to add an additional cool drawings
+			// this will let use draw images, graphs, and a bunch of other cool visualizers
+			//ActiveProperty.Property->DrawDebug(Canvas);
+		}					
+	}
+}
+
+void FInputDeviceDebugTools::OnShowDebugHardwareDevices(UCanvas* Canvas)
+{
+	const UInputDeviceSubsystem* System = UInputDeviceSubsystem::Get();
+	FDisplayDebugManager& DisplayDebugManager = Canvas->DisplayDebugManager;
+
+	DisplayDebugManager.SetFont(GEngine->GetMediumFont());
+
+	if (!System)
+	{
+		DisplayDebugManager.SetDrawColor(FColor::Red);
+		DisplayDebugManager.DrawString(TEXT("UInputDeviceSubsystem is not initalized!"));
+		return;
+	}	
+
+	IPlatformInputDeviceMapper& DeviceMapper = IPlatformInputDeviceMapper::Get();
+	
+	// Title
+	DisplayDebugManager.SetDrawColor(FColor::Orange);
+	DisplayDebugManager.DrawString(TEXT("Input Devices"));
+
+	TArray<FPlatformUserId> AllUsers;
+	DeviceMapper.GetAllActiveUsers(OUT AllUsers);
+
+	// Info relevant to all platform users
+	const int32 NumUsers = DeviceMapper.GetAllActiveUsers(AllUsers);
+
+	DisplayDebugManager.SetDrawColor(FColor::White);
+	DisplayDebugManager.DrawString(FString::Printf(TEXT("Total Platform Users: %d"), NumUsers));
+
+	for (const FPlatformUserId UserId : AllUsers)
+	{
+		TArray<FInputDeviceId> AllDevices;
+		const int32 NumDevices = DeviceMapper.GetAllInputDevicesForUser(UserId, OUT AllDevices);
+
+		// Show how many input devices each user has
+		DisplayDebugManager.SetDrawColor(NumDevices > 0 ? FColor::Green : FColor::Orange);
+		DisplayDebugManager.DrawString(FString::Printf(TEXT("Platform User '%d' has '%d' input devices"), UserId.GetInternalId(), NumDevices));
+
+		for (const FInputDeviceId DeviceId : AllDevices)
+		{
+			static const float InputDeviceLabelXOffset = 10.0f;
+			static const float InputDeviceInfoXOffset = InputDeviceLabelXOffset + 10.0f;
+
+			// Print some information about the known hardware for this input device
+			const FHardwareDeviceIdentifier Hardware = System->GetInputDeviceHardwareIdentifier(DeviceId);
+			{
+				const FString HardwareDeviceInfo = FString::Printf(TEXT("Input Device ID %d \t (%s::%s)"),
+					DeviceId.GetId(),
+					*Hardware.InputClassName.ToString(),
+					*Hardware.HardwareDeviceIdentifier.ToString());			
+
+				DisplayDebugManager.SetDrawColor(DeviceId.IsValid() ? FColor::White : FColor::Red);
+				DisplayDebugManager.DrawString(HardwareDeviceInfo, InputDeviceLabelXOffset);
+			}
+			
+			// Show info about how this input device is seen by Slate
+			if (FSlateApplication::IsInitialized())
+			{
+				TOptional<int32> SlateUserId = FSlateApplication::Get().GetUserIndexForInputDevice(DeviceId);					
+				if (SlateUserId.IsSet())
+				{
+					DisplayDebugManager.SetDrawColor(FColor::White);
+					DisplayDebugManager.DrawString(FString::Printf(TEXT("Slate User Index: %d"), SlateUserId.GetValue()), InputDeviceInfoXOffset);
+				}
+				else
+				{
+					DisplayDebugManager.SetDrawColor(FColor::Red);
+					DisplayDebugManager.DrawString(FString::Printf(TEXT("Slate User Index: INVALID")), InputDeviceInfoXOffset);
+				}
+			}
+
+			// Display the connection status of the input device
+			{
+				const EInputDeviceConnectionState DeviceState = DeviceMapper.GetInputDeviceConnectionState(DeviceId);
+
+				static const TMap<EInputDeviceConnectionState, FColor> ConnectionColorMap = 
+				{
+					{ EInputDeviceConnectionState::Invalid,			FColor::Red },
+					{ EInputDeviceConnectionState::Unknown,			FColor::Orange },
+					{ EInputDeviceConnectionState::Disconnected,	FColor::Silver },
+					{ EInputDeviceConnectionState::Connected,		FColor::Green },
+				};
+
+				static const TMap<EInputDeviceConnectionState, FString> ConnectionLabelMap =
+				{
+					{ EInputDeviceConnectionState::Invalid,			TEXT("INVALID") },
+					{ EInputDeviceConnectionState::Unknown,			TEXT("Unknown") },
+					{ EInputDeviceConnectionState::Disconnected,	TEXT("Disconnected") },
+					{ EInputDeviceConnectionState::Connected,		TEXT("Connected") },
+				};
+
+				if (const FColor* Color = ConnectionColorMap.Find(DeviceState))
+				{
+					DisplayDebugManager.SetDrawColor(*Color);
+				}
+
+				if (const FString* Label = ConnectionLabelMap.Find(DeviceState))
+				{
+					DisplayDebugManager.DrawString(*Label, InputDeviceInfoXOffset);
+				}			
+			}
+		}		
+	}
+}
+
+void FInputDeviceDebugTools::ListAllKnownHardwareDeviceIdentifier(const TArray<FString>& Args, UWorld* World)
+{
+	if (UInputPlatformSettings* System = UInputPlatformSettings::Get())
+	{
+		const TArray<FHardwareDeviceIdentifier>& Devices = System->GetHardwareDevices();
+		for (const FHardwareDeviceIdentifier& Device : Devices)
+		{
+			UE_LOG(LogInputDeviceProperties, Log, TEXT("%s::%s"), *Device.InputClassName.ToString(), *Device.HardwareDeviceIdentifier.ToString());
+		}
+	}
+}
+
+#endif	// SUPPORT_INPUT_DEVICE_DEBUGGING
\ No newline at end of file
diff --git a/Engine/Plugins/Runtime/InputDebugging/Source/InputDebugging/Private/InputDeviceDebugTools.h b/Engine/Plugins/Runtime/InputDebugging/Source/InputDebugging/Private/InputDeviceDebugTools.h
new file mode 100644
--- /dev/null
+++ b/Engine/Plugins/Runtime/InputDebugging/Source/InputDebugging/Private/InputDeviceDebugTools.h
@@ -0,0 +1,52 @@
+// Copyright Epic Games, Inc. All Rights Reserved.
+
+#pragma once
+
+#include "Misc/Build.h"
+
+// Allow you to enable this debug drawing in a build.cs if you want, but default to when we are not in shipping.
+#if !defined(SUPPORT_INPUT_DEVICE_DEBUGGING)
+#define SUPPORT_INPUT_DEVICE_DEBUGGING 		!UE_BUILD_SHIPPING
+#endif
+
+#if SUPPORT_INPUT_DEVICE_DEBUGGING
+
+#include "Templates/SharedPointer.h"
+
+class UCanvas;
+struct IConsoleCommand;
+class UWorld;
+class AHUD;
+class FDebugDisplayInfo;
+
+/**  */
+class FInputDeviceDebugTools : public TSharedFromThis<FInputDeviceDebugTools>
+{
+public:
+
+	FInputDeviceDebugTools();
+	~FInputDeviceDebugTools();
+	
+	// non-copyable, non-movable
+	FInputDeviceDebugTools(const FInputDeviceDebugTools& Other) = delete;
+	FInputDeviceDebugTools(FInputDeviceDebugTools&& Other) = delete;
+
+private:
+
+	void AddConsoleCommands();
+	void RemoveConsoleCommands();
+
+	static void OnShowDebugInfo(AHUD* HUD, UCanvas* Canvas, const FDebugDisplayInfo& DisplayInfo, float& YL, float& YPos);
+
+	static void OnShowDebugDeviceProperties(UCanvas* Canvas);
+
+	static void OnShowDebugHardwareDevices(UCanvas* Canvas);
+
+	/** Log all the platform's currently available FHardwareDeviceIdentifier  */
+	void ListAllKnownHardwareDeviceIdentifier(const TArray<FString>& Args, UWorld* World);
+
+	/** Console commands registered by this device debugger */
+	TArray<IConsoleCommand*> ConsoleCommands;
+};
+
+#endif // SUPPORT_INPUT_DEVICE_DEBUGGING
\ No newline at end of file
diff --git a/Engine/Source/Runtime/Engine/Classes/GameFramework/InputDeviceSubsystem.h b/Engine/Source/Runtime/Engine/Classes/GameFramework/InputDeviceSubsystem.h
--- a/Engine/Source/Runtime/Engine/Classes/GameFramework/InputDeviceSubsystem.h
+++ b/Engine/Source/Runtime/Engine/Classes/GameFramework/InputDeviceSubsystem.h
@@ -134,7 +134,7 @@ UCLASS(BlueprintType)
 class ENGINE_API UInputDeviceSubsystem : public UEngineSubsystem, public FTickableGameObject
 {
 	friend class FInputDeviceSubsystemProcessor;
-	friend class FInputDeviceDebugVisualizer;
+	friend class FInputDeviceDebugTools;
 	
 	GENERATED_BODY()
 
