Add a bHasBeenAppliedAtLeastOnce state to active device properties.

This handles cases where the evaluation time is longer then the duration of a property. I.e., your property is set to a duration of 0.1, but you get a delta time of .12 seconds for some reason (choppy frames, low perf client, etc).

#rnx
#pf skip

diff --git a/Engine/Plugins/Runtime/InputDebugging/Source/InputDebugging/Private/InputDeviceDebugTools.cpp b/Engine/Plugins/Runtime/InputDebugging/Source/InputDebugging/Private/InputDeviceDebugTools.cpp
--- a/Engine/Plugins/Runtime/InputDebugging/Source/InputDebugging/Private/InputDeviceDebugTools.cpp
+++ b/Engine/Plugins/Runtime/InputDebugging/Source/InputDebugging/Private/InputDeviceDebugTools.cpp
@@ -122,8 +122,8 @@ void FInputDeviceDebugTools::OnShowDebugDeviceProperties(UCanvas* Canvas)
 			// Draw some information about the flags that have been set on the given property
 			DisplayDebugManager.SetDrawColor(FColor::Silver);
 			DisplayDebugManager.DrawString(FString::Printf(TEXT("\t Property Handle: %s"), *ActiveProperty->PropertyHandle.ToString()));
-			DisplayDebugManager.DrawString(FString::Printf(TEXT("\t Flags: \n \t \t \t bLooping: %d \n \t \t \t bIgnoreTimeDilation: %d \n \t \t \t bPlayWhilePaused: %d"),
-				ActiveProperty->bLooping, ActiveProperty->bIgnoreTimeDilation, ActiveProperty->bPlayWhilePaused));
+			DisplayDebugManager.DrawString(FString::Printf(TEXT("\t Flags: \n \t \t \t bLooping: %d \n \t \t \t bIgnoreTimeDilation: %d \n \t \t \t bPlayWhilePaused: %d \n \t \t \t bHasBeenAppliedAtLeastOnce: %d"),
+				ActiveProperty->bLooping, ActiveProperty->bIgnoreTimeDilation, ActiveProperty->bPlayWhilePaused, ActiveProperty->bHasBeenAppliedAtLeastOnce));
 
 			// Highlight if a property's evaluation has completed. If can go higher then the duration if
 			// is set to not be removed after its evaluation is complete
diff --git a/Engine/Source/Runtime/Engine/Classes/GameFramework/InputDeviceProperties.h b/Engine/Source/Runtime/Engine/Classes/GameFramework/InputDeviceProperties.h
--- a/Engine/Source/Runtime/Engine/Classes/GameFramework/InputDeviceProperties.h
+++ b/Engine/Source/Runtime/Engine/Classes/GameFramework/InputDeviceProperties.h
@@ -133,7 +133,7 @@ protected:
 	* to be the max time range.
 	*/
 	UPROPERTY(VisibleAnywhere, BlueprintReadWrite, Category = "Info")
-	float PropertyDuration = 0.1f;
+	float PropertyDuration = 0.0f;
 };
 
 template<class TDataLayout>
diff --git a/Engine/Source/Runtime/Engine/Classes/GameFramework/InputDeviceSubsystem.h b/Engine/Source/Runtime/Engine/Classes/GameFramework/InputDeviceSubsystem.h
--- a/Engine/Source/Runtime/Engine/Classes/GameFramework/InputDeviceSubsystem.h
+++ b/Engine/Source/Runtime/Engine/Classes/GameFramework/InputDeviceSubsystem.h
@@ -108,6 +108,15 @@ protected:
 
 		/** If true, then this device property will be played even if the game world is paused. */
 		bool bPlayWhilePaused = false;
+
+		/**
+		 * This is set to true when this device property has been applied.
+		 * All device properties should be applied at least one time, no matter what.
+		 * This handles cases where the evaluation time is longer then the duration of a property.
+		 * I.e., your property is set to a duration of 0.1, but you get a delta time of .12 seconds for some reason
+		 * (choppy frames, low perf client, debugging, etc). 
+		 */
+		bool bHasBeenAppliedAtLeastOnce = false;
 	};
 
 public:
diff --git a/Engine/Source/Runtime/Engine/Private/GameFramework/InputDeviceSubsystem.cpp b/Engine/Source/Runtime/Engine/Private/GameFramework/InputDeviceSubsystem.cpp
--- a/Engine/Source/Runtime/Engine/Private/GameFramework/InputDeviceSubsystem.cpp
+++ b/Engine/Source/Runtime/Engine/Private/GameFramework/InputDeviceSubsystem.cpp
@@ -280,7 +280,7 @@ void UInputDeviceSubsystem::Tick(float InDeltaTime)
 
 		// If the property has run past it's duration, reset it and remove it from our active properties
 		// Don't remove properties that are looping, because the user desires them to keep playing
-		if (!ActiveProp.bLooping && ActiveProp.EvaluatedDuration > ActiveProp.Property->GetDuration())
+		if (!ActiveProp.bLooping && ActiveProp.EvaluatedDuration > ActiveProp.Property->GetDuration() && ActiveProp.bHasBeenAppliedAtLeastOnce)
 		{
 			ActiveProp.Property->ResetDeviceProperty(ActiveProp.PlatformUser);
 			PropertiesPendingRemoval.Remove(ActiveProp.PropertyHandle);
@@ -292,6 +292,9 @@ void UInputDeviceSubsystem::Tick(float InDeltaTime)
 		{
 			ActiveProp.Property->EvaluateDeviceProperty(ActiveProp.PlatformUser, DeltaTime, ActiveProp.EvaluatedDuration);
 			ActiveProp.Property->ApplyDeviceProperty(ActiveProp.PlatformUser);
+
+			// Track that this property has now successfully been applied at least once
+			ActiveProp.bHasBeenAppliedAtLeastOnce = true;
 		}
 	}
 
