Fix another pop up that would happen incorrectly when you run the engine without a project.

#rnx

diff --git a/Engine/Plugins/Runtime/InputDebugging/Source/InputDebuggingEditor/Private/InputDebuggingEditorModule.cpp b/Engine/Plugins/Runtime/InputDebugging/Source/InputDebuggingEditor/Private/InputDebuggingEditorModule.cpp
--- a/Engine/Plugins/Runtime/InputDebugging/Source/InputDebuggingEditor/Private/InputDebuggingEditorModule.cpp
+++ b/Engine/Plugins/Runtime/InputDebugging/Source/InputDebuggingEditor/Private/InputDebuggingEditorModule.cpp
@@ -33,6 +33,11 @@ public:
 
 void FInputDebuggingEditorModule::StartupModule()
 {
+	if (!FApp::HasProjectName())
+	{
+		return;
+	}
+
 	// Listen for when input devices change
 	IPlatformInputDeviceMapper::Get().GetOnInputDeviceConnectionChange().AddRaw(this, &FInputDebuggingEditorModule::OnInputDeviceConnectionChange);
 	
@@ -43,6 +48,11 @@ void FInputDebuggingEditorModule::StartupModule()
 
 void FInputDebuggingEditorModule::ShutdownModule()
 {
+	if (!FApp::HasProjectName())
+	{
+		return;
+	}
+
 	// Remove any input listeners
 	IPlatformInputDeviceMapper::Get().GetOnInputDeviceConnectionChange().RemoveAll(this);
 	
