Fix a non-unity compile error from FApp::HasProjectName

#rnx

diff --git a/Engine/Plugins/Runtime/InputDebugging/Source/InputDebuggingEditor/Private/InputDebuggingEditorModule.cpp b/Engine/Plugins/Runtime/InputDebugging/Source/InputDebuggingEditor/Private/InputDebuggingEditorModule.cpp
--- a/Engine/Plugins/Runtime/InputDebugging/Source/InputDebuggingEditor/Private/InputDebuggingEditorModule.cpp
+++ b/Engine/Plugins/Runtime/InputDebugging/Source/InputDebuggingEditor/Private/InputDebuggingEditorModule.cpp
@@ -9,6 +9,7 @@
 #include "Styling/SlateStyleMacros.h"
 #include "Styling/SlateStyleRegistry.h"
 #include "Misc/ConfigCacheIni.h"
+#include "Misc/App.h"
 
 #define LOCTEXT_NAMESPACE "InputDebuggingEditor"
 
