Readded SlateCore to fix CIS job.

diff --git a/Engine/Plugins/Runtime/InputDebugging/Source/InputDebugging/InputDebugging.Build.cs b/Engine/Plugins/Runtime/InputDebugging/Source/InputDebugging/InputDebugging.Build.cs
--- a/Engine/Plugins/Runtime/InputDebugging/Source/InputDebugging/InputDebugging.Build.cs
+++ b/Engine/Plugins/Runtime/InputDebugging/Source/InputDebugging/InputDebugging.Build.cs
@@ -12,8 +12,9 @@ namespace UnrealBuildTool.Rules
 					"Core",
 					"CoreUObject",
 					"Engine",
-					"Slate",
 					"InputCore",
+					"Slate",
+					"SlateCore",
 				}
 			);
 		}
