[Input] Add a "Input.LogAllConnectedDevices" command to dump all the currently connected input devices and some metadata bout them to the logs.

diff --git a/Engine/Plugins/Runtime/InputDebugging/Source/InputDebugging/Private/InputDeviceDebugTools.cpp b/Engine/Plugins/Runtime/InputDebugging/Source/InputDebugging/Private/InputDeviceDebugTools.cpp
--- a/Engine/Plugins/Runtime/InputDebugging/Source/InputDebugging/Private/InputDeviceDebugTools.cpp
+++ b/Engine/Plugins/Runtime/InputDebugging/Source/InputDebugging/Private/InputDeviceDebugTools.cpp
@@ -32,6 +32,13 @@ void FInputDeviceDebugTools::AddConsoleCommands()
 		FConsoleCommandWithWorldAndArgsDelegate::CreateRaw(this, &FInputDeviceDebugTools::ListAllKnownHardwareDeviceIdentifier),
 		ECVF_Cheat
 	));
+
+	ConsoleCommands.Add(IConsoleManager::Get().RegisterConsoleCommand(
+		TEXT("Input.LogAllConnectedDevices"),
+		TEXT("Log all the platform's currently connected devices from the device mapper"),
+		FConsoleCommandWithWorldAndArgsDelegate::CreateRaw(this, &FInputDeviceDebugTools::LogAllConnectedDevices),
+		ECVF_Cheat
+	));
 }
 
 void FInputDeviceDebugTools::RemoveConsoleCommands()
@@ -253,4 +260,42 @@ void FInputDeviceDebugTools::ListAllKnownHardwareDeviceIdentifier(const TArray<F
 	}
 }
 
+void FInputDeviceDebugTools::LogAllConnectedDevices(const TArray<FString>& Args, UWorld* World)
+{
+	const UInputDeviceSubsystem* System = UInputDeviceSubsystem::Get();
+
+	if (!System)
+	{
+		UE_LOG(LogInput, Error, TEXT("[%hs] 'Input.LogAllConnectedDevices' command failed! Unable to find a valid UInputDeviceSubsystem!"), __func__);
+		return;
+	}
+
+	IPlatformInputDeviceMapper& DeviceMapper = IPlatformInputDeviceMapper::Get();
+
+	TArray<FPlatformUserId> AllUsers;
+	DeviceMapper.GetAllActiveUsers(OUT AllUsers);
+
+	// Info relevant to all platform users
+	const int32 NumUsers = DeviceMapper.GetAllActiveUsers(AllUsers);
+
+	FString LogMessage = FString::Printf(TEXT("Total number of active users: %d\n"), NumUsers);
+
+	for (const FPlatformUserId UserId : AllUsers)
+	{
+		TArray<FInputDeviceId> AllDevices;
+		const int32 NumDevices = DeviceMapper.GetAllInputDevicesForUser(UserId, OUT AllDevices);
+
+		LogMessage.Append(FString::Printf(TEXT("Platform User: %d\t\t %d devices\n"), UserId.GetInternalId(), AllDevices.Num()));
+
+		for (const FInputDeviceId DeviceId : AllDevices)
+		{
+			const FHardwareDeviceIdentifier DeviceData = System->GetInputDeviceHardwareIdentifier(DeviceId);
+
+			LogMessage.Append(FString::Printf(TEXT("\tInput Device Id: %d    Hardware Info: %s\n"), DeviceId.GetId(), *DeviceData.ToString()));
+		}
+	}
+
+	UE_LOG(LogInput, Log, TEXT("[%hs] Logging all connected input devices from 'Input.LogAllConnectedDevices' command... \n%s"), __func__, *LogMessage);
+}
+
 #endif	// SUPPORT_INPUT_DEVICE_DEBUGGING
diff --git a/Engine/Plugins/Runtime/InputDebugging/Source/InputDebugging/Private/InputDeviceDebugTools.h b/Engine/Plugins/Runtime/InputDebugging/Source/InputDebugging/Private/InputDeviceDebugTools.h
--- a/Engine/Plugins/Runtime/InputDebugging/Source/InputDebugging/Private/InputDeviceDebugTools.h
+++ b/Engine/Plugins/Runtime/InputDebugging/Source/InputDebugging/Private/InputDeviceDebugTools.h
@@ -45,6 +45,8 @@ private:
 	/** Log all the platform's currently available FHardwareDeviceIdentifier  */
 	void ListAllKnownHardwareDeviceIdentifier(const TArray<FString>& Args, UWorld* World);
 
+	void LogAllConnectedDevices(const TArray<FString>& Args, UWorld* World);
+
 	/** Console commands registered by this device debugger */
 	TArray<IConsoleCommand*> ConsoleCommands;
 };
