[Input] Add an editor preference for setting a "Simulated Device Mapping Policy". This can be used in the editor to change how input devices get mapped to platform users, which can be useful for testing device connection sequences for your UI, log in flows, etc.

Add the current device mapping policy to the output of the "Input.LogAllConnectedDevices" command.

diff --git a/Engine/Plugins/Runtime/InputDebugging/Source/InputDebugging/InputDebugging.Build.cs b/Engine/Plugins/Runtime/InputDebugging/Source/InputDebugging/InputDebugging.Build.cs
--- a/Engine/Plugins/Runtime/InputDebugging/Source/InputDebugging/InputDebugging.Build.cs
+++ b/Engine/Plugins/Runtime/InputDebugging/Source/InputDebugging/InputDebugging.Build.cs
@@ -17,6 +17,15 @@ namespace UnrealBuildTool.Rules
 					"SlateCore",
 				}
 			);
+			
+			if (Target.bBuildEditor)
+			{
+				PublicDependencyModuleNames.AddRange(
+					new string[] {
+						"UnrealEd"
+					}
+				);
+			}
 		}
 	}
 }
diff --git a/Engine/Plugins/Runtime/InputDebugging/Source/InputDebugging/Private/InputDeviceDebugTools.cpp b/Engine/Plugins/Runtime/InputDebugging/Source/InputDebugging/Private/InputDeviceDebugTools.cpp
--- a/Engine/Plugins/Runtime/InputDebugging/Source/InputDebugging/Private/InputDeviceDebugTools.cpp
+++ b/Engine/Plugins/Runtime/InputDebugging/Source/InputDebugging/Private/InputDeviceDebugTools.cpp
@@ -4,12 +4,16 @@
 
 #if SUPPORT_INPUT_DEVICE_DEBUGGING
 
+#include "DrawDebugHelpers.h"
+#include "Engine/Canvas.h"
+#include "Engine/Engine.h"
 #include "Framework/Application/SlateApplication.h"
-#include "GameFramework/InputDeviceProperties.h"
 #include "GameFramework/HUD.h"
-#include "Engine/Engine.h"
-#include "Engine/Canvas.h"
-#include "DrawDebugHelpers.h"
+#include "GameFramework/InputDeviceProperties.h"
+
+#if WITH_EDITOR
+#include "Settings/LevelEditorPlaySettings.h"
+#endif
 
 FInputDeviceDebugTools::FInputDeviceDebugTools()
 {
@@ -17,6 +21,30 @@ FInputDeviceDebugTools::FInputDeviceDebugTools()
 	AHUD::OnShowDebugInfo.AddStatic(&FInputDeviceDebugTools::OnShowDebugInfo);
 #endif
 	AddConsoleCommands();
+
+#if WITH_EDITOR
+	if (ULevelEditorPlaySettings* Settings = GetMutableDefault<ULevelEditorPlaySettings>())
+	{
+		Settings->SimulatedDeviceMappingPolicyPostEditChange.AddLambda([]()
+		{
+			const ULevelEditorPlaySettings* EditorSettings = GetDefault<ULevelEditorPlaySettings>();
+			// Notify the device mapper that we have changed and have it re-map input devices
+			if (!EditorSettings->IsSimulateDeviceMappingPolicy())
+			{
+				// change it back to whatever the default project setting is
+				if (const UInputPlatformSettings* Settings = UInputPlatformSettings::Get())
+				{
+					IPlatformInputDeviceMapper::Get().HandleInputDevicePolicyChanged(Settings->DeviceMappingPolicy);
+				}
+			}
+			// If the policy is not invalid, then have the device mapper change to whatever it is now set to
+			else
+			{
+				IPlatformInputDeviceMapper::Get().HandleInputDevicePolicyChanged(EditorSettings->GetSimulatedDeviceMappingPolicy());
+			}
+		});
+	}
+#endif	// WITH_EDITOR
 }
 
 FInputDeviceDebugTools::~FInputDeviceDebugTools()
@@ -279,6 +307,26 @@ void FInputDeviceDebugTools::LogAllConnectedDevices(const TArray<FString>& Args,
 	const int32 NumUsers = DeviceMapper.GetAllActiveUsers(AllUsers);
 
 	FString LogMessage = FString::Printf(TEXT("Total number of active users: %d\n"), NumUsers);
+	
+	{
+		const EInputDeviceMappingPolicy MappingPolicy = IPlatformInputDeviceMapper::Get().GetCurrentDeviceMappingPolicy();
+		
+		auto PolicyToString = [](const EInputDeviceMappingPolicy Policy)-> FString
+		{
+			switch (Policy)
+			{
+			case EInputDeviceMappingPolicy::Invalid: return TEXT("Invalid");
+			case EInputDeviceMappingPolicy::UseManagedPlatformLogin: return TEXT("UseManagedPlatformLogin");
+			case EInputDeviceMappingPolicy::PrimaryUserSharesKeyboardAndFirstGamepad: return TEXT("PrimaryUserSharesKeyboardAndFirstGamepad");
+			case EInputDeviceMappingPolicy::CreateUniquePlatformUserForEachDevice: return TEXT("CreateUniquePlatformUserForEachDevice");
+			case EInputDeviceMappingPolicy::MapAllDevicesToPrimaryUser: return TEXT("MapAllDevicesToPrimaryUser");
+			}
+
+			return TEXT("Unknown EInputDeviceMappingPolicy");
+		};
+		
+		LogMessage.Append(FString::Printf(TEXT("Input Device Mapping Policy: %s\n"), *PolicyToString(MappingPolicy)));
+	}
 
 	for (const FPlatformUserId UserId : AllUsers)
 	{
diff --git a/Engine/Source/Editor/UnrealEd/Classes/Settings/LevelEditorPlaySettings.h b/Engine/Source/Editor/UnrealEd/Classes/Settings/LevelEditorPlaySettings.h
--- a/Engine/Source/Editor/UnrealEd/Classes/Settings/LevelEditorPlaySettings.h
+++ b/Engine/Source/Editor/UnrealEd/Classes/Settings/LevelEditorPlaySettings.h
@@ -11,6 +11,7 @@
 #include "Misc/App.h"
 #include "Settings/LevelEditorPlayNetworkEmulationSettings.h"
 #include "ToolMenuContext.h"
+#include "GenericPlatform/InputDeviceMappingPolicy.h"
 #include "LevelEditorPlaySettings.generated.h"
 
 class SWindow;
@@ -231,6 +232,27 @@ public:
 	UPROPERTY(config, EditAnywhere, Category=PlayInEditor)
 	bool UseMouseForTouch;
 
+protected:
+	
+	UPROPERTY(config, EditAnywhere, Category=PlayInEditor, meta = (InlineEditConditionToggle))
+	bool bSimulateDeviceMappingPolicy = false;
+	
+	/**
+	 * If enabled, then this simulated device mapping policy will be used instead of the one set in the Project Settings.
+	 * Simulating this device mapping policy may be useful to test different controller connection and login flows
+	 * in the editor which would otherwise require you to deploy to a console to test.
+	 */
+	UPROPERTY(config, EditAnywhere, Category=PlayInEditor, meta=(EditCondition=bSimulateDeviceMappingPolicy))
+	EInputDeviceMappingPolicy SimulatedDeviceMappingPolicy = EInputDeviceMappingPolicy::UseManagedPlatformLogin;
+
+public:
+
+	bool IsSimulateDeviceMappingPolicy() const { return bSimulateDeviceMappingPolicy; }
+
+	EInputDeviceMappingPolicy GetSimulatedDeviceMappingPolicy() const { return SimulatedDeviceMappingPolicy; }
+
+	TMulticastDelegate<void()> SimulatedDeviceMappingPolicyPostEditChange;
+	
 	/** Whether to show a label for mouse control gestures in the PIE view. */
 	UPROPERTY(config, EditAnywhere, Category=PlayInEditor)
 	bool ShowMouseControlLabel;
diff --git a/Engine/Source/Editor/UnrealEd/Private/Settings/SettingsClasses.cpp b/Engine/Source/Editor/UnrealEd/Private/Settings/SettingsClasses.cpp
--- a/Engine/Source/Editor/UnrealEd/Private/Settings/SettingsClasses.cpp
+++ b/Engine/Source/Editor/UnrealEd/Private/Settings/SettingsClasses.cpp
@@ -10,6 +10,7 @@
 #include "AI/NavigationSystemBase.h"
 #include "Model.h"
 #include "ISourceControlModule.h"
+#include "GameFramework/InputSettings.h"
 #include "Settings/ContentBrowserSettings.h"
 #include "Settings/LevelEditorPlaySettings.h"
 #include "Settings/LevelEditorViewportSettings.h"
@@ -554,7 +555,9 @@ void ULevelEditorPlaySettings::PostEditChangeProperty(struct FPropertyChangedEve
 		BuildGameBeforeLaunch = EPlayOnBuildMode::PlayOnBuild_Never;
 	}
 
-	if (PropertyChangedEvent.Property && PropertyChangedEvent.Property->GetFName() == GET_MEMBER_NAME_CHECKED(ULevelEditorPlaySettings, bOnlyLoadVisibleLevelsInPIE))
+	const FName PropertyName = PropertyChangedEvent.Property ? PropertyChangedEvent.Property->GetFName() : NAME_None;
+
+	if (PropertyName == GET_MEMBER_NAME_CHECKED(ULevelEditorPlaySettings, bOnlyLoadVisibleLevelsInPIE))
 	{
 		for (TObjectIterator<UWorld> WorldIt; WorldIt; ++WorldIt)
 		{
@@ -567,6 +570,11 @@ void ULevelEditorPlaySettings::PostEditChangeProperty(struct FPropertyChangedEve
 		NetworkEmulationSettings.OnPostEditChange(PropertyChangedEvent);
 	}
 
+	if (PropertyName == GET_MEMBER_NAME_CHECKED(ULevelEditorPlaySettings, SimulatedDeviceMappingPolicy) || PropertyName == GET_MEMBER_NAME_CHECKED(ULevelEditorPlaySettings, bSimulateDeviceMappingPolicy))
+	{
+		SimulatedDeviceMappingPolicyPostEditChange.Broadcast();
+	}
+	
 	PushDebugDrawingSettings();
 	if (PropertyChangedEvent.MemberProperty && PropertyChangedEvent.MemberProperty->GetFName() == GET_MEMBER_NAME_CHECKED(ULevelEditorPlaySettings, bShowServerDebugDrawingByDefault))
 	{
