[Input] Display the device type name in the controller connection toast in the editor.

#rnx

#Synced-CL 50735006

diff --git a/Engine/Plugins/Runtime/InputDebugging/Source/InputDebuggingEditor/InputDebuggingEditor.Build.cs b/Engine/Plugins/Runtime/InputDebugging/Source/InputDebuggingEditor/InputDebuggingEditor.Build.cs
--- a/Engine/Plugins/Runtime/InputDebugging/Source/InputDebuggingEditor/InputDebuggingEditor.Build.cs
+++ b/Engine/Plugins/Runtime/InputDebugging/Source/InputDebuggingEditor/InputDebuggingEditor.Build.cs
@@ -8,11 +8,12 @@ namespace UnrealBuildTool.Rules
 		{
 			PrivateDependencyModuleNames.AddRange(
 				new string[] {
-					"Core",
 					"ApplicationCore",
-					"Slate",
-					"SlateCore",
+					"Core",
+					"Engine",
 					"MainFrame",
+					"Slate",
+					"SlateCore",				 					
 				}
 			);
 		}
diff --git a/Engine/Plugins/Runtime/InputDebugging/Source/InputDebuggingEditor/Private/InputDebuggingEditorModule.cpp b/Engine/Plugins/Runtime/InputDebugging/Source/InputDebuggingEditor/Private/InputDebuggingEditorModule.cpp
--- a/Engine/Plugins/Runtime/InputDebugging/Source/InputDebuggingEditor/Private/InputDebuggingEditorModule.cpp
+++ b/Engine/Plugins/Runtime/InputDebugging/Source/InputDebuggingEditor/Private/InputDebuggingEditorModule.cpp
@@ -3,6 +3,7 @@
 #include "InputDebuggingEditorModule.h"
 #include "GenericPlatform/GenericPlatformInputDeviceMapper.h"
 #include "Framework/Notifications/NotificationManager.h"
+#include "GameFramework/InputSettings.h"
 #include "Widgets/Notifications/SNotificationList.h"
 #include "Interfaces/IMainFrameModule.h"
 #include "Styling/SlateStyle.h"
@@ -131,7 +132,28 @@ void FInputDebuggingEditorModule::OnInputDeviceConnectionChange(EInputDeviceConn
 	FFormatNamedArguments SubMessageArgs;
 	SubMessageArgs.Add(TEXT("PlatformUser"), PlatformUserId.GetInternalId());
 	SubMessageArgs.Add(TEXT("DeviceId"), InputDeviceId.GetId());
-	const FText SubMessage = FText::Format(LOCTEXT("InputDeviceConnectionChangeSubMessage", "Platform User ID: {PlatformUser}\nDevice ID: {DeviceId}"), SubMessageArgs);
+
+	{
+		// Not all systems will have a device scope available right away upon connection yet, and we need that to proceed 
+		if (const FInputDeviceScope* Scope = FInputDeviceScope::GetCurrent())
+		{
+			const FName DeviceScopeName(*Scope->HardwareDeviceIdentifier);
+			if (const FHardwareDeviceIdentifier* Hardware = UInputPlatformSettings::Get()->GetHardwareDeviceForClassName(DeviceScopeName))
+			{
+				SubMessageArgs.Add(TEXT("HardwareDeviceID"), FText::FromString(Hardware->ToString()));
+			}
+			else
+			{
+				SubMessageArgs.Add(TEXT("HardwareDeviceID"), FText::FromString(Scope->HardwareDeviceIdentifier));
+			}
+		}
+		else
+		{
+			SubMessageArgs.Add(TEXT("HardwareDeviceID"), LOCTEXT("InputDeviceConnectionChangeSubMessage_Label", "Unknown"));
+		}
+	}
+	
+	const FText SubMessage = FText::Format(LOCTEXT("InputDeviceConnectionChangeSubMessage", "Platform User ID: {PlatformUser}\nDevice ID: {DeviceId}\nType: {HardwareDeviceID}"), SubMessageArgs);
 
 	// Send out a slate notification that there was an input device connection change
 	FNotificationInfo* Info = new FNotificationInfo(Message);
