[Input debugger] Clean up the delegate call just use the input device subsystem instead of manually checking the device scope.

#rnx

#Synced-CL 50969376

diff --git a/Engine/Plugins/Runtime/InputDebugging/Source/InputDebuggingEditor/Private/InputDebuggingEditorModule.cpp b/Engine/Plugins/Runtime/InputDebugging/Source/InputDebuggingEditor/Private/InputDebuggingEditorModule.cpp
--- a/Engine/Plugins/Runtime/InputDebugging/Source/InputDebuggingEditor/Private/InputDebuggingEditorModule.cpp
+++ b/Engine/Plugins/Runtime/InputDebugging/Source/InputDebuggingEditor/Private/InputDebuggingEditorModule.cpp
@@ -3,6 +3,7 @@
 #include "InputDebuggingEditorModule.h"
 #include "GenericPlatform/GenericPlatformInputDeviceMapper.h"
 #include "Framework/Notifications/NotificationManager.h"
+#include "GameFramework/InputDeviceSubsystem.h"
 #include "GameFramework/InputSettings.h"
 #include "Widgets/Notifications/SNotificationList.h"
 #include "Interfaces/IMainFrameModule.h"
@@ -132,26 +133,9 @@ void FInputDebuggingEditorModule::OnInputDeviceConnectionChange(EInputDeviceConn
 	FFormatNamedArguments SubMessageArgs;
 	SubMessageArgs.Add(TEXT("PlatformUser"), PlatformUserId.GetInternalId());
 	SubMessageArgs.Add(TEXT("DeviceId"), InputDeviceId.GetId());
-
-	{
-		// Not all systems will have a device scope available right away upon connection yet, and we need that to proceed 
-		if (const FInputDeviceScope* Scope = FInputDeviceScope::GetCurrent())
-		{
-			const FName DeviceScopeName(*Scope->HardwareDeviceIdentifier);
-			if (const FHardwareDeviceIdentifier* Hardware = UInputPlatformSettings::Get()->GetHardwareDeviceForClassName(DeviceScopeName))
-			{
-				SubMessageArgs.Add(TEXT("HardwareDeviceID"), FText::FromString(Hardware->ToString()));
-			}
-			else
-			{
-				SubMessageArgs.Add(TEXT("HardwareDeviceID"), FText::FromString(Scope->HardwareDeviceIdentifier));
-			}
-		}
-		else
-		{
-			SubMessageArgs.Add(TEXT("HardwareDeviceID"), LOCTEXT("InputDeviceConnectionChangeSubMessage_Label", "Unknown"));
-		}
-	}
+	
+	const FHardwareDeviceIdentifier Hardware = UInputDeviceSubsystem::Get()->GetInputDeviceHardwareIdentifier(InputDeviceId);
+	SubMessageArgs.Add(TEXT("HardwareDeviceID"), FText::FromString(Hardware.ToString()));
 	
 	const FText SubMessage = FText::Format(LOCTEXT("InputDeviceConnectionChangeSubMessage", "Platform User ID: {PlatformUser}\nDevice ID: {DeviceId}\nType: {HardwareDeviceID}"), SubMessageArgs);
 
